More compile fixes.  PointCloudSceneProxy temporarily broken due to like of GetElementShaderBindings implementation

diff --git a/Engine/Plugins/Experimental/AR/PointCloud/Source/PointCloud/Private/PointCloudSceneProxy.cpp b/Engine/Plugins/Experimental/AR/PointCloud/Source/PointCloud/Private/PointCloudSceneProxy.cpp
--- a/Engine/Plugins/Experimental/AR/PointCloud/Source/PointCloud/Private/PointCloudSceneProxy.cpp
+++ b/Engine/Plugins/Experimental/AR/PointCloud/Source/PointCloud/Private/PointCloudSceneProxy.cpp
@@ -170,11 +170,11 @@ void FPointCloudSceneProxy::GetDynamicMeshElements(const TArray<const FSceneView
 	}
 
 	const bool bIsWireframe = AllowDebugViewmodes() && ViewFamily.EngineShowFlags.Wireframe;
-	FMaterialRenderProxy* MaterialProxy = Material->GetRenderProxy(IsSelected());
+	FMaterialRenderProxy* MaterialProxy = Material->GetRenderProxy();
 	if (bIsWireframe)
 	{
 		FMaterialRenderProxy* WireframeMaterialInstance = new FColoredMaterialRenderProxy(
-				GEngine->WireframeMaterial ? GEngine->WireframeMaterial->GetRenderProxy(IsSelected()) : nullptr,
+				GEngine->WireframeMaterial->GetRenderProxy(),
 				FLinearColor(0, 0.5f, 1.f)
 		);
 		Collector.RegisterOneFrameMaterialProxy(WireframeMaterialInstance);
diff --git a/Engine/Plugins/Experimental/AR/PointCloud/Source/PointCloud/Private/PointCloudVertexFactory.cpp b/Engine/Plugins/Experimental/AR/PointCloud/Source/PointCloud/Private/PointCloudVertexFactory.cpp
--- a/Engine/Plugins/Experimental/AR/PointCloud/Source/PointCloud/Private/PointCloudVertexFactory.cpp
+++ b/Engine/Plugins/Experimental/AR/PointCloud/Source/PointCloud/Private/PointCloudVertexFactory.cpp
@@ -29,6 +29,20 @@ public:
 		Ar << PointSize;
 	}
 
+	virtual void GetElementShaderBindings(
+		const class FSceneInterface* Scene,
+		const class FSceneView* View,
+		const class FMeshMaterialShader* Shader,
+		bool bShaderRequiresPositionOnlyStream,
+		ERHIFeatureLevel::Type FeatureLevel,
+		const class FVertexFactory* VertexFactory,
+		const struct FMeshBatchElement& BatchElement,
+		class FMeshDrawSingleShaderBindings& ShaderBindings,
+		FVertexInputStreamArray& VertexStreams) const override
+	{		
+	}
+
+#if 0
 	virtual void SetMesh(FRHICommandList& RHICmdList, FShader* Shader, const FVertexFactory* InVertexFactory, const FSceneView& View, const FMeshBatchElement& BatchElement, uint32 DataFlags) const override
 	{
 		FPointCloudVertexFactory* VertexFactory = (FPointCloudVertexFactory*)InVertexFactory;
@@ -39,6 +53,7 @@ public:
 		SetShaderValue(RHICmdList, VertexShaderRHI, ColorMask, VertexFactory->GetColorMask());
 		SetShaderValue(RHICmdList, VertexShaderRHI, PointSize, VertexFactory->GetPointSize());
 	}
+#endif
 
 	virtual uint32 GetSize() const override { return sizeof(*this); }
 
diff --git a/Engine/Plugins/Experimental/RemoteSession/Source/RemoteSession/Private/Channels/RemoteSessionARCameraChannel.cpp b/Engine/Plugins/Experimental/RemoteSession/Source/RemoteSession/Private/Channels/RemoteSessionARCameraChannel.cpp
--- a/Engine/Plugins/Experimental/RemoteSession/Source/RemoteSession/Private/Channels/RemoteSessionARCameraChannel.cpp
+++ b/Engine/Plugins/Experimental/RemoteSession/Source/RemoteSession/Private/Channels/RemoteSessionARCameraChannel.cpp
@@ -316,7 +316,7 @@ void FARCameraSceneViewExtension::RenderARCamera_RenderThread(FRHICommandListImm
 
 	SetUniformBufferParameterImmediate(RHICmdList, VertexShader->GetVertexShader(), VertexShader->GetUniformBufferParameter<FDrawRectangleParameters>(), Parameters);
 	VertexShader->SetParameters(RHICmdList, InView);
-	PixelShader->SetParameters(RHICmdList, InView, PPMaterial->GetRenderProxy(false));
+	PixelShader->SetParameters(RHICmdList, InView, PPMaterial->GetRenderProxy());
 
 	if (VertexBufferRHI && IndexBufferRHI.IsValid())
 	{
diff --git a/Engine/Plugins/Runtime/Oculus/OculusVR/Source/OculusMR/Private/OculusMR_BoundaryMeshComponent.cpp b/Engine/Plugins/Runtime/Oculus/OculusVR/Source/OculusMR/Private/OculusMR_BoundaryMeshComponent.cpp
--- a/Engine/Plugins/Runtime/Oculus/OculusVR/Source/OculusMR/Private/OculusMR_BoundaryMeshComponent.cpp
+++ b/Engine/Plugins/Runtime/Oculus/OculusVR/Source/OculusMR/Private/OculusMR_BoundaryMeshComponent.cpp
@@ -114,7 +114,7 @@ public:
 		const bool bWireframe = AllowDebugViewmodes() && ViewFamily.EngineShowFlags.Wireframe;
 
 		auto WireframeMaterialInstance = new FColoredMaterialRenderProxy(
-			GEngine->WireframeMaterial ? GEngine->WireframeMaterial->GetRenderProxy(IsSelected()) : NULL,
+			GEngine->WireframeMaterial->GetRenderProxy(),
 			FLinearColor(0, 0.5f, 1.f)
 		);
 
@@ -127,7 +127,7 @@ public:
 		}
 		else
 		{
-			MaterialProxy = Material->GetRenderProxy(IsSelected());
+			MaterialProxy = Material->GetRenderProxy();
 		}
 
 		for (int32 ViewIndex = 0; ViewIndex < Views.Num(); ViewIndex++)
diff --git a/Engine/Plugins/Runtime/Oculus/OculusVR/Source/OculusMR/Private/OculusMR_PlaneMeshComponent.cpp b/Engine/Plugins/Runtime/Oculus/OculusVR/Source/OculusMR/Private/OculusMR_PlaneMeshComponent.cpp
--- a/Engine/Plugins/Runtime/Oculus/OculusVR/Source/OculusMR/Private/OculusMR_PlaneMeshComponent.cpp
+++ b/Engine/Plugins/Runtime/Oculus/OculusVR/Source/OculusMR/Private/OculusMR_PlaneMeshComponent.cpp
@@ -92,7 +92,7 @@ public:
 			const bool bWireframe = AllowDebugViewmodes() && ViewFamily.EngineShowFlags.Wireframe;
 
 			auto WireframeMaterialInstance = new FColoredMaterialRenderProxy(
-				GEngine->WireframeMaterial ? GEngine->WireframeMaterial->GetRenderProxy(IsSelected()) : NULL,
+				GEngine->WireframeMaterial->GetRenderProxy(),
 				FLinearColor(0, 0.5f, 1.f)
 			);
 
@@ -105,7 +105,7 @@ public:
 			}
 			else
 			{
-				MaterialProxy = Material->GetRenderProxy(IsSelected());
+				MaterialProxy = Material->GetRenderProxy();
 			}
 
 			for (int32 ViewIndex = 0; ViewIndex < Views.Num(); ViewIndex++)
