Fixing integration of FPointCloudVertexFactoryShaderParameters.

diff --git a/Engine/Plugins/Experimental/AR/PointCloud/Source/PointCloud/Private/PointCloudVertexFactory.cpp b/Engine/Plugins/Experimental/AR/PointCloud/Source/PointCloud/Private/PointCloudVertexFactory.cpp
--- a/Engine/Plugins/Experimental/AR/PointCloud/Source/PointCloud/Private/PointCloudVertexFactory.cpp
+++ b/Engine/Plugins/Experimental/AR/PointCloud/Source/PointCloud/Private/PointCloudVertexFactory.cpp
@@ -7,6 +7,7 @@
 #include "Engine/Engine.h"
 #include "Materials/Material.h"
 #include "ShaderParameterUtils.h"
+#include "MeshMaterialShader.h"
 
 IMPLEMENT_GLOBAL_SHADER_PARAMETER_STRUCT(FPointCloudVertexFactoryParameters, "PointCloudVF");
 
@@ -35,25 +36,18 @@ public:
 		const class FMeshMaterialShader* Shader,
 		bool bShaderRequiresPositionOnlyStream,
 		ERHIFeatureLevel::Type FeatureLevel,
-		const class FVertexFactory* VertexFactory,
+		const class FVertexFactory* InVertexFactory,
 		const struct FMeshBatchElement& BatchElement,
 		class FMeshDrawSingleShaderBindings& ShaderBindings,
 		FVertexInputStreamArray& VertexStreams) const override
 	{		
-	}
-
-#if 0
-	virtual void SetMesh(FRHICommandList& RHICmdList, FShader* Shader, const FVertexFactory* InVertexFactory, const FSceneView& View, const FMeshBatchElement& BatchElement, uint32 DataFlags) const override
-	{
 		FPointCloudVertexFactory* VertexFactory = (FPointCloudVertexFactory*)InVertexFactory;
-		FVertexShaderRHIParamRef VertexShaderRHI = Shader->GetVertexShader();
 
-		SetUniformBufferParameter(RHICmdList, VertexShaderRHI, Shader->GetUniformBufferParameter<FPointCloudVertexFactoryParameters>(), VertexFactory->GetPointCloudVertexFactoryUniformBuffer());
+		ShaderBindings.Add(Shader->GetUniformBufferParameter<FPointCloudVertexFactoryParameters>(), VertexFactory->GetPointCloudVertexFactoryUniformBuffer());
 
-		SetShaderValue(RHICmdList, VertexShaderRHI, ColorMask, VertexFactory->GetColorMask());
-		SetShaderValue(RHICmdList, VertexShaderRHI, PointSize, VertexFactory->GetPointSize());
+		ShaderBindings.Add(ColorMask, VertexFactory->GetColorMask());
+		ShaderBindings.Add(PointSize, VertexFactory->GetPointSize());
 	}
-#endif
 
 	virtual uint32 GetSize() const override { return sizeof(*this); }
 
