Fixed non-unity compile errors

diff --git a/Engine/Plugins/Experimental/AR/PointCloud/Source/PointCloud/Private/PointCloudComponent.cpp b/Engine/Plugins/Experimental/AR/PointCloud/Source/PointCloud/Private/PointCloudComponent.cpp
--- a/Engine/Plugins/Experimental/AR/PointCloud/Source/PointCloud/Private/PointCloudComponent.cpp
+++ b/Engine/Plugins/Experimental/AR/PointCloud/Source/PointCloud/Private/PointCloudComponent.cpp
@@ -3,9 +3,14 @@
 #include "PointCloudComponent.h"
 #include "PointCloud.h"
 #include "PointCloudSceneProxy.h"
-
-DECLARE_CYCLE_STAT(TEXT("Create Point Cloud Proxy"), STAT_PointCloud_CreateSceneProxy, STATGROUP_POINTCLOUD);
-DECLARE_CYCLE_STAT(TEXT("Point Cloud Comp Update"), STAT_PointCloud_ComponentUpdateCost, STATGROUP_POINTCLOUD);
+#include "UObject/ConstructorHelpers.h"
+#include "EngineGlobals.h"
+#include "Engine/Engine.h"
+#include "Materials/MaterialInterface.h"
+#include "Materials/Material.h"
+
+DECLARE_CYCLE_STAT(TEXT("Create Point Cloud Proxy"), STAT_PointCloud_CreateSceneProxy, STATGROUP_PointCloud);
+DECLARE_CYCLE_STAT(TEXT("Point Cloud Comp Update"), STAT_PointCloud_ComponentUpdateCost, STATGROUP_PointCloud);
 
 UPointCloudComponent::UPointCloudComponent(const FObjectInitializer& ObjectInitializer)
 	: Super(ObjectInitializer)
diff --git a/Engine/Plugins/Experimental/AR/PointCloud/Source/PointCloud/Private/PointCloudModule.cpp b/Engine/Plugins/Experimental/AR/PointCloud/Source/PointCloud/Private/PointCloudModule.cpp
--- a/Engine/Plugins/Experimental/AR/PointCloud/Source/PointCloud/Private/PointCloudModule.cpp
+++ b/Engine/Plugins/Experimental/AR/PointCloud/Source/PointCloud/Private/PointCloudModule.cpp
@@ -1,6 +1,7 @@
 // Copyright 1998-2018 Epic Games, Inc. All Rights Reserved.
 
 #include "PointCloud.h"
+#include "Modules/ModuleManager.h"
 
 DEFINE_LOG_CATEGORY(LogPointCloud);
 
diff --git a/Engine/Plugins/Experimental/AR/PointCloud/Source/PointCloud/Private/PointCloudSceneProxy.cpp b/Engine/Plugins/Experimental/AR/PointCloud/Source/PointCloud/Private/PointCloudSceneProxy.cpp
--- a/Engine/Plugins/Experimental/AR/PointCloud/Source/PointCloud/Private/PointCloudSceneProxy.cpp
+++ b/Engine/Plugins/Experimental/AR/PointCloud/Source/PointCloud/Private/PointCloudSceneProxy.cpp
@@ -18,9 +18,9 @@
 #include "SceneManagement.h"
 #include "DynamicMeshBuilder.h"
 
-DECLARE_CYCLE_STAT(TEXT("Update Point Cloud GT"), STAT_PointCloud_UpdatePointCloud, STATGROUP_POINTCLOUD);
-DECLARE_CYCLE_STAT(TEXT("Get Mesh Elements"), STAT_PointCloud_GetMeshElements, STATGROUP_POINTCLOUD);
-DECLARE_CYCLE_STAT(TEXT("Create RT Resources"), STAT_PointCloud_CreateRenderThreadResources, STATGROUP_POINTCLOUD);
+DECLARE_CYCLE_STAT(TEXT("Update Point Cloud GT"), STAT_PointCloud_UpdatePointCloud, STATGROUP_PointCloud);
+DECLARE_CYCLE_STAT(TEXT("Get Mesh Elements"), STAT_PointCloud_GetMeshElements, STATGROUP_PointCloud);
+DECLARE_CYCLE_STAT(TEXT("Create RT Resources"), STAT_PointCloud_CreateRenderThreadResources, STATGROUP_PointCloud);
 
 
 FNoFetchPointCloudSceneProxy::FNoFetchPointCloudSceneProxy(UPointCloudComponent* Component)
diff --git a/Engine/Plugins/Experimental/AR/PointCloud/Source/PointCloud/Private/PointCloudSceneProxy.h b/Engine/Plugins/Experimental/AR/PointCloud/Source/PointCloud/Private/PointCloudSceneProxy.h
--- a/Engine/Plugins/Experimental/AR/PointCloud/Source/PointCloud/Private/PointCloudSceneProxy.h
+++ b/Engine/Plugins/Experimental/AR/PointCloud/Source/PointCloud/Private/PointCloudSceneProxy.h
@@ -3,6 +3,11 @@
 #pragma once
 
 #include "CoreMinimal.h"
+#include "PrimitiveViewRelevance.h"
+#include "PrimitiveSceneProxy.h"
+#include "Materials/MaterialInterface.h"
+#include "MaterialShared.h"
+#include "MeshBatch.h"
 #include "PrimitiveSceneProxy.h"
 #include "PointCloudVertexFactory.h"
 #include "PointCloudBuffers.h"
diff --git a/Engine/Plugins/Experimental/AR/PointCloud/Source/PointCloud/Private/PointCloudVertexFactory.h b/Engine/Plugins/Experimental/AR/PointCloud/Source/PointCloud/Private/PointCloudVertexFactory.h
--- a/Engine/Plugins/Experimental/AR/PointCloud/Source/PointCloud/Private/PointCloudVertexFactory.h
+++ b/Engine/Plugins/Experimental/AR/PointCloud/Source/PointCloud/Private/PointCloudVertexFactory.h
@@ -2,6 +2,8 @@
 #pragma once
 
 #include "CoreMinimal.h"
+#include "RHI.h"
+#include "RenderResource.h"
 #include "UniformBuffer.h"
 #include "VertexFactory.h"
 
diff --git a/Engine/Plugins/Experimental/AR/PointCloud/Source/PointCloud/Public/PointCloud.h b/Engine/Plugins/Experimental/AR/PointCloud/Source/PointCloud/Public/PointCloud.h
--- a/Engine/Plugins/Experimental/AR/PointCloud/Source/PointCloud/Public/PointCloud.h
+++ b/Engine/Plugins/Experimental/AR/PointCloud/Source/PointCloud/Public/PointCloud.h
@@ -3,7 +3,9 @@
 #pragma once
 
 #include "CoreMinimal.h"
+#include "Stats/Stats.h"
+#include "Logging/LogMacros.h"
 
 DECLARE_LOG_CATEGORY_EXTERN(LogPointCloud, Log, All);
 
-DECLARE_STATS_GROUP(TEXT("Point Cloud"), STATGROUP_POINTCLOUD, STATCAT_Advanced);
+DECLARE_STATS_GROUP(TEXT("Point Cloud"), STATGROUP_PointCloud, STATCAT_Advanced);
