Fixed a typo in a comment

diff --git a/Engine/Plugins/Experimental/AR/PointCloud/Source/PointCloud/Private/PointCloudComponent.cpp b/Engine/Plugins/Experimental/AR/PointCloud/Source/PointCloud/Private/PointCloudComponent.cpp
--- a/Engine/Plugins/Experimental/AR/PointCloud/Source/PointCloud/Private/PointCloudComponent.cpp
+++ b/Engine/Plugins/Experimental/AR/PointCloud/Source/PointCloud/Private/PointCloudComponent.cpp
@@ -20,7 +20,7 @@ UPointCloudComponent::UPointCloudComponent(const FObjectInitializer& ObjectIniti
 	PrimaryComponentTick.bCanEverTick = true;
 	PrimaryComponentTick.TickGroup = TG_PostPhysics;
 
-	// Load our material we use f
+	// Load our material we use for rendering
 	static ConstructorHelpers::FObjectFinder<UMaterialInterface> DefaultMaterial(TEXT("/PointCloud/DefaultPointCloudMaterial"));
 	PointCloudMaterial = DefaultMaterial.Object;
 	if (PointCloudMaterial == nullptr)
