Removing "ShaderCore" module reference as it is now "RenderCore"

diff --git a/Engine/Plugins/Experimental/AR/PointCloud/Source/PointCloud/PointCloud.Build.cs b/Engine/Plugins/Experimental/AR/PointCloud/Source/PointCloud/PointCloud.Build.cs
--- a/Engine/Plugins/Experimental/AR/PointCloud/Source/PointCloud/PointCloud.Build.cs
+++ b/Engine/Plugins/Experimental/AR/PointCloud/Source/PointCloud/PointCloud.Build.cs
@@ -28,7 +28,6 @@ namespace UnrealBuildTool.Rules
                     "Core",
                     "Engine",
 					"RenderCore",
-					"ShaderCore",
 					"RHI",
                     "AugmentedReality",
                     // ... add other public dependencies that you statically link with here ...
