Fix for PointCloudVertexFactory to deal with new shader API

diff --git a/Engine/Plugins/Experimental/AR/PointCloud/Source/PointCloud/Private/PointCloudVertexFactory.cpp b/Engine/Plugins/Experimental/AR/PointCloud/Source/PointCloud/Private/PointCloudVertexFactory.cpp
--- a/Engine/Plugins/Experimental/AR/PointCloud/Source/PointCloud/Private/PointCloudVertexFactory.cpp
+++ b/Engine/Plugins/Experimental/AR/PointCloud/Source/PointCloud/Private/PointCloudVertexFactory.cpp
@@ -8,7 +8,7 @@
 #include "Materials/Material.h"
 #include "ShaderParameterUtils.h"
 
-IMPLEMENT_UNIFORM_BUFFER_STRUCT(FPointCloudVertexFactoryParameters, TEXT("PointCloudVF"));
+IMPLEMENT_GLOBAL_SHADER_PARAMETER_STRUCT(FPointCloudVertexFactoryParameters, "PointCloudVF");
 
 /**
  * Shader parameters for the vector field visualization vertex factory.
diff --git a/Engine/Plugins/Experimental/AR/PointCloud/Source/PointCloud/Private/PointCloudVertexFactory.h b/Engine/Plugins/Experimental/AR/PointCloud/Source/PointCloud/Private/PointCloudVertexFactory.h
--- a/Engine/Plugins/Experimental/AR/PointCloud/Source/PointCloud/Private/PointCloudVertexFactory.h
+++ b/Engine/Plugins/Experimental/AR/PointCloud/Source/PointCloud/Private/PointCloudVertexFactory.h
@@ -10,10 +10,10 @@
 /**
  * Uniform buffer to hold parameters for point cloud rendering
  */
-BEGIN_UNIFORM_BUFFER_STRUCT( FPointCloudVertexFactoryParameters, )
-	UNIFORM_MEMBER_SRV(Buffer<float>, VertexFetch_PointLocationBuffer)
-	UNIFORM_MEMBER_SRV(Buffer<float4>, VertexFetch_PointColorBuffer)
-END_UNIFORM_BUFFER_STRUCT( FPointCloudVertexFactoryParameters )
+BEGIN_GLOBAL_SHADER_PARAMETER_STRUCT( FPointCloudVertexFactoryParameters, )
+	SHADER_PARAMETER_SRV(Buffer<float>, VertexFetch_PointLocationBuffer)
+	SHADER_PARAMETER_SRV(Buffer<float4>, VertexFetch_PointColorBuffer)
+END_GLOBAL_SHADER_PARAMETER_STRUCT()
 typedef TUniformBufferRef<FPointCloudVertexFactoryParameters> FPointCloudVertexFactoryBufferRef;
 
 /**
