Make VF params type declaration non-inline

diff --git a/Engine/Plugins/Experimental/AR/PointCloud/Source/PointCloud/Private/PointCloudVertexFactory.cpp b/Engine/Plugins/Experimental/AR/PointCloud/Source/PointCloud/Private/PointCloudVertexFactory.cpp
--- a/Engine/Plugins/Experimental/AR/PointCloud/Source/PointCloud/Private/PointCloudVertexFactory.cpp
+++ b/Engine/Plugins/Experimental/AR/PointCloud/Source/PointCloud/Private/PointCloudVertexFactory.cpp
@@ -17,7 +17,7 @@ IMPLEMENT_GLOBAL_SHADER_PARAMETER_STRUCT(FPointCloudVertexFactoryParameters, "Po
 class FPointCloudVertexFactoryShaderParameters :
 	public FVertexFactoryShaderParameters
 {
-	DECLARE_INLINE_TYPE_LAYOUT(FPointCloudVertexFactoryShaderParameters, NonVirtual);
+	DECLARE_TYPE_LAYOUT(FPointCloudVertexFactoryShaderParameters, NonVirtual);
 public:
 	void Bind(const FShaderParameterMap& ParameterMap)
 	{
@@ -134,5 +134,6 @@ void FPointCloudVertexFactory::SetParameters(const FPointCloudVertexFactoryParam
 	PointSize = InSize;
 }
 
+IMPLEMENT_TYPE_LAYOUT(FPointCloudVertexFactoryShaderParameters);
 IMPLEMENT_VERTEX_FACTORY_PARAMETER_TYPE(FPointCloudVertexFactory, SF_Vertex, FPointCloudVertexFactoryShaderParameters);
 IMPLEMENT_VERTEX_FACTORY_TYPE(FPointCloudVertexFactory, "/Engine/Private/PointCloudVertexFactory.ush", true, false, false, false, false);
