Fixed build break.

diff --git a/Engine/Plugins/Experimental/AR/PointCloud/Source/PointCloud/Private/PointCloudSceneProxy.cpp b/Engine/Plugins/Experimental/AR/PointCloud/Source/PointCloud/Private/PointCloudSceneProxy.cpp
--- a/Engine/Plugins/Experimental/AR/PointCloud/Source/PointCloud/Private/PointCloudSceneProxy.cpp
+++ b/Engine/Plugins/Experimental/AR/PointCloud/Source/PointCloud/Private/PointCloudSceneProxy.cpp
@@ -138,21 +138,22 @@ void FPointCloudSceneProxy::CreateRenderThreadResources()
 {
 	SCOPE_CYCLE_COUNTER(STAT_PointCloud_CreateRenderThreadResources);
 
-	PointCloudVertexFactory.InitResource(FRHICommandListImmediate::Get());
-	PointCloudIndexBuffer.InitRHIWithSize(Points.Num());
+	FRHICommandListBase& RHICmdList = FRHICommandListImmediate::Get();
+	PointCloudVertexFactory.InitResource(RHICmdList);
+	PointCloudIndexBuffer.InitRHIWithSize(RHICmdList, Points.Num());
 
 	// We either use a single color or the color array
 	if (Colors.Num())
 	{
-		PointCloudColorVertexBuffer.InitRHIWith(Colors);
+		PointCloudColorVertexBuffer.InitRHIWith(RHICmdList, Colors);
 		Colors.Empty();
 	}
 	else
 	{
-		PointCloudColorVertexBuffer.InitRHIWith(PointColor.ToFColor(false));
+		PointCloudColorVertexBuffer.InitRHIWith(RHICmdList, PointColor.ToFColor(false));
 	}
 
-	PointCloudLocationVertexBuffer.InitRHIWith(Points);
+	PointCloudLocationVertexBuffer.InitRHIWith(RHICmdList, Points);
 	Points.Empty();
 
 	// Setup the vertex factory shader parameters
